Add forward declare to fix nounity CIS error

diff --git a/Engine/Source/Runtime/AudioCodecEngine/Public/DecoderInputFactory.h b/Engine/Source/Runtime/AudioCodecEngine/Public/DecoderInputFactory.h
--- a/Engine/Source/Runtime/AudioCodecEngine/Public/DecoderInputFactory.h
+++ b/Engine/Source/Runtime/AudioCodecEngine/Public/DecoderInputFactory.h
@@ -6,6 +6,7 @@
 
 // Forward declares.
 class USoundWave;
+class FSoundWaveProxy;
 
 namespace Audio
 {
