Fix Non-unity compile error

diff --git a/Engine/Source/Runtime/AudioCodecEngine/Public/DecoderInputFactory.h b/Engine/Source/Runtime/AudioCodecEngine/Public/DecoderInputFactory.h
--- a/Engine/Source/Runtime/AudioCodecEngine/Public/DecoderInputFactory.h
+++ b/Engine/Source/Runtime/AudioCodecEngine/Public/DecoderInputFactory.h
@@ -14,6 +14,7 @@ namespace Audio
 {
 	// Forward declares.
 	struct IDecoderInput;
+	using FSoundWaveProxyPtr = TSharedPtr<FSoundWaveProxy, ESPMode::ThreadSafe>;
 
 	// Just loose for now.
 	AUDIOCODECENGINE_API TUniquePtr<IDecoderInput> CreateBackCompatDecoderInput(
