IWYU Fix

diff --git a/Engine/Source/Runtime/AudioCodecEngine/Public/DecoderInputFactory.h b/Engine/Source/Runtime/AudioCodecEngine/Public/DecoderInputFactory.h
--- a/Engine/Source/Runtime/AudioCodecEngine/Public/DecoderInputFactory.h
+++ b/Engine/Source/Runtime/AudioCodecEngine/Public/DecoderInputFactory.h
@@ -2,6 +2,8 @@
 
 #pragma once
 
+#include "Templates/SharedPointer.h"
+#include "Templates/SharedPointerInternals.h"
 #include "Templates/UniquePtr.h"
 #include "AudioMixer.h"
 
@@ -17,7 +19,5 @@ namespace Audio
 	using FSoundWaveProxyPtr = TSharedPtr<FSoundWaveProxy, ESPMode::ThreadSafe>;
 
 	// Just loose for now.
-	AUDIOCODECENGINE_API TUniquePtr<IDecoderInput> CreateBackCompatDecoderInput(
-		FName InOldFormatName,
-		const FSoundWaveProxyPtr& InSoundWave );
+	AUDIOCODECENGINE_API TUniquePtr<IDecoderInput> CreateBackCompatDecoderInput(FName InOldFormatName, const FSoundWaveProxyPtr& InSoundWave);
 }
