Fix crash when playing a metasound in PIE after using multiple client windows.

diff --git a/Engine/Source/Runtime/AudioCodecEngine/Private/DecoderInputBackCompat.cpp b/Engine/Source/Runtime/AudioCodecEngine/Private/DecoderInputBackCompat.cpp
--- a/Engine/Source/Runtime/AudioCodecEngine/Private/DecoderInputBackCompat.cpp
+++ b/Engine/Source/Runtime/AudioCodecEngine/Private/DecoderInputBackCompat.cpp
@@ -17,38 +17,43 @@ namespace Audio
 		{
 			FAudioDeviceHandle Handle = FAudioDeviceManager::Get()->GetActiveAudioDevice();
 
-			OldInfoObject.Reset(Handle->CreateCompressedAudioInfo(Wave));
-			audio_ensure(OldInfoObject.IsValid());
-
-			FSoundQualityInfo Info;
-			if (Wave->IsStreaming())
+			if (ensure(Handle))
 			{
-				if (!OldInfoObject->StreamCompressedInfo(Wave, &Info))
+				OldInfoObject.Reset(Handle->CreateCompressedAudioInfo(Wave));
+				audio_ensure(OldInfoObject.IsValid());
+
+				FSoundQualityInfo Info;
+				if (Wave->IsStreaming())
 				{
-					return nullptr;
+					if (!OldInfoObject->StreamCompressedInfo(Wave, &Info))
+					{
+						return nullptr;
+					}
 				}
-			}
-			else
-			{
-				if (!OldInfoObject->ReadCompressedInfo(Wave->GetResourceData(), Wave->GetResourceSize(), &Info))
+				else
 				{
-					return nullptr;
+					if (!OldInfoObject->ReadCompressedInfo(Wave->GetResourceData(), Wave->GetResourceSize(), &Info))
+					{
+						return nullptr;
+					}
 				}
-			}
 
-			Desc.NumChannels		= Info.NumChannels;
-			Desc.NumFramesPerSec	= Info.SampleRate;
-			Desc.NumFrames			= (uint32)((float)Info.Duration * Info.SampleRate);
-			Desc.NumBytesPerPacket	= ~0;
+				Desc.NumChannels		= Info.NumChannels;
+				Desc.NumFramesPerSec	= Info.SampleRate;
+				Desc.NumFrames			= (uint32)((float)Info.Duration * Info.SampleRate);
+				Desc.NumBytesPerPacket	= ~0;
 
-			Desc.CodecName			= FBackCompatCodec::GetDetailsStatic().Name;
-			Desc.CodecFamilyName	= FBackCompatCodec::GetDetailsStatic().FamilyName;
-			Desc.CodecVersion		= FBackCompatCodec::GetDetailsStatic().Version;		
+				Desc.CodecName			= FBackCompatCodec::GetDetailsStatic().Name;
+				Desc.CodecFamilyName	= FBackCompatCodec::GetDetailsStatic().FamilyName;
+				Desc.CodecVersion		= FBackCompatCodec::GetDetailsStatic().Version;		
+			}
 		}
+
 		if( OutDescriptor )
 		{		
 			*OutDescriptor = Desc;
 		}		
+
 		return OldInfoObject.Get();
 	}
 
diff --git a/Engine/Source/Runtime/Engine/Private/AudioDeviceManager.cpp b/Engine/Source/Runtime/Engine/Private/AudioDeviceManager.cpp
--- a/Engine/Source/Runtime/Engine/Private/AudioDeviceManager.cpp
+++ b/Engine/Source/Runtime/Engine/Private/AudioDeviceManager.cpp
@@ -883,7 +883,12 @@ FAudioDeviceHandle FAudioDeviceManager::GetActiveAudioDevice()
 {
 	if (ActiveAudioDeviceID != INDEX_NONE)
 	{
-		return GetAudioDevice(ActiveAudioDeviceID);
+		FAudioDeviceHandle ActiveAudioDeviceHandle = GetAudioDevice(ActiveAudioDeviceID);
+
+		if (ActiveAudioDeviceHandle)
+		{
+			return ActiveAudioDeviceHandle;
+		}
 	}
 
 	return MainAudioDeviceHandle;
