Fix for ensure when changing format on a SoundWave that are forceinline.

diff --git a/Engine/Source/Runtime/AudioCodecEngine/Private/DecoderInputBackCompat.cpp b/Engine/Source/Runtime/AudioCodecEngine/Private/DecoderInputBackCompat.cpp
--- a/Engine/Source/Runtime/AudioCodecEngine/Private/DecoderInputBackCompat.cpp
+++ b/Engine/Source/Runtime/AudioCodecEngine/Private/DecoderInputBackCompat.cpp
@@ -14,26 +14,40 @@ namespace Audio
 		FFormatDescriptorSection* OutDescriptor /*= nullptr*/) const
 	{
 		if (!OldInfoObject.IsValid())
-		{			
-			OldInfoObject.Reset(IAudioInfoFactoryRegistry::Get().Create(Wave->GetRuntimeFormat()));
-			audio_ensure(OldInfoObject.IsValid());
+		{
+			FAudioDeviceHandle Handle = FAudioDeviceManager::Get()->GetActiveAudioDevice();
+			FName Format = Wave->GetRuntimeFormat();
+			IAudioInfoFactory* Factory = IAudioInfoFactoryRegistry::Get().Find(Format);
+			if (ensure(Handle))
+			{
+				return nullptr;
+			}
+
+				OldInfoObject.Reset(Handle->CreateCompressedAudioInfo(Wave));
+			if (!ensure(InfoInstance.IsValid()))
+			{
+				return nullptr;
+			}
 
 			FSoundQualityInfo Info;
 			if (Wave->IsStreaming())
 			{
-				if (!OldInfoObject->StreamCompressedInfo(Wave, &Info))
+				if (!InfoInstance->StreamCompressedInfo(Wave, &Info))
 				{
 					return nullptr;
 				}
 			}
 			else
 			{
-				if (!OldInfoObject->ReadCompressedInfo(Wave->GetResourceData(), Wave->GetResourceSize(), &Info))
+				if (!InfoInstance->ReadCompressedInfo(Wave->GetResourceData(), Wave->GetResourceSize(), &Info))
 				{
 					return nullptr;
 				}
 			}
 
+			// Commit the new instance only if we successfully read the info above.
+			OldInfoObject.Reset(InfoInstance.Release());
+
 			Desc.NumChannels		= Info.NumChannels;
 			Desc.NumFramesPerSec	= Info.SampleRate;
 			Desc.NumFrames			= (uint32)((float)Info.Duration * Info.SampleRate);
@@ -41,13 +55,13 @@ namespace Audio
 
 			Desc.CodecName			= FBackCompatCodec::GetDetailsStatic().Name;
 			Desc.CodecFamilyName	= FBackCompatCodec::GetDetailsStatic().FamilyName;
-			Desc.CodecVersion		= FBackCompatCodec::GetDetailsStatic().Version;		
+			Desc.CodecVersion		= FBackCompatCodec::GetDetailsStatic().Version;				
 		}
 
-		if( OutDescriptor )
-		{		
+		if (OutDescriptor)
+		{
 			*OutDescriptor = Desc;
-		}		
+		}
 
 		return OldInfoObject.Get();
 	}
diff --git a/Engine/Source/Runtime/Engine/Private/SoundWave.cpp b/Engine/Source/Runtime/Engine/Private/SoundWave.cpp
--- a/Engine/Source/Runtime/Engine/Private/SoundWave.cpp
+++ b/Engine/Source/Runtime/Engine/Private/SoundWave.cpp
@@ -1619,10 +1619,10 @@ void USoundWave::InvalidateCompressedData(bool bFreeResources, bool bRebuildStre
 	check(SoundWaveDataPtr);
 
 	CompressedDataGuid = FGuid::NewGuid();
+
 	SoundWaveDataPtr->DiscardZerothChunkData();
 	SoundWaveDataPtr->CompressedFormatData.FlushData();
-	SoundWaveDataPtr->ResourceSize = 0;
-	SoundWaveDataPtr->ResourceData.Empty();
+	RemoveAudioResource();
 
 	if (bFreeResources)
 	{
