Fix for frame counter uint rollover when decoding wave in metasounds

diff --git a/Engine/Source/Runtime/AudioCodecEngine/Private/DecoderBackCompat.cpp b/Engine/Source/Runtime/AudioCodecEngine/Private/DecoderBackCompat.cpp
--- a/Engine/Source/Runtime/AudioCodecEngine/Private/DecoderBackCompat.cpp
+++ b/Engine/Source/Runtime/AudioCodecEngine/Private/DecoderBackCompat.cpp
@@ -51,7 +51,7 @@ namespace Audio
 
 	Audio::IDecoder::FDecodeReturn FBackCompat::Decode(bool bLoop)
 	{
-		uint32 NumFramesRemaining = Reqs.NumSampleFramesWanted;
+		int32 NumFramesRemaining = Reqs.NumSampleFramesWanted;
 		FBackCompatInput& BackCompatSrc = static_cast<FBackCompatInput&>(*Src);
 		ICompressedAudioInfo* Info = BackCompatSrc.GetInfo();
 
@@ -105,7 +105,7 @@ namespace Audio
 			);
 
 			FrameOffset += NumFramesStreamed;
-			NumFramesRemaining -= NumFramesStreamed;
+			NumFramesRemaining -= FMath::Min(NumFramesStreamed, NumFramesRemaining);
 		}
 
 		if (!bFinished)
@@ -122,4 +122,4 @@ namespace Audio
 		}
 	}
 
-}
\ No newline at end of file
+}
