Fix missing exported functions
#horde 341705
#rnx

diff --git a/Engine/Source/Runtime/AudioCodecEngine/Public/SoundWaveProxyReader.h b/Engine/Source/Runtime/AudioCodecEngine/Public/SoundWaveProxyReader.h
--- a/Engine/Source/Runtime/AudioCodecEngine/Public/SoundWaveProxyReader.h
+++ b/Engine/Source/Runtime/AudioCodecEngine/Public/SoundWaveProxyReader.h
@@ -49,24 +49,24 @@ namespace SoundWaveProxyReaderPrivate
 		void Reset();
 
 		/** Returns requirements used by the audio codec system. */
-		virtual Audio::IDecoderOutput::FRequirements GetRequirements(const Audio::FDecodedFormatInfo & InFormat) const override;
+		AUDIOCODECENGINE_API virtual Audio::IDecoderOutput::FRequirements GetRequirements(const Audio::FDecodedFormatInfo & InFormat) const override;
 
 		/** Adds samples to the buffer.
 		 *
 		 * This is called by the Decoder and should not be called otherwise.
 		 */
-		virtual int32 PushAudio(const Audio::IDecoderOutput::FPushedAudioDetails & InDetails, TArrayView<const int16> In16BitInterleave) override;
+		AUDIOCODECENGINE_API virtual int32 PushAudio(const Audio::IDecoderOutput::FPushedAudioDetails & InDetails, TArrayView<const int16> In16BitInterleave) override;
 
 		/** Adds samples to the buffer.
 		 *
 		 * This is called by the Decoder and should not be called otherwise.
 		 */
-		virtual int32 PushAudio(const FPushedAudioDetails & InDetails, TArrayView<const float> InFloat32Interleave) override;
+		AUDIOCODECENGINE_API virtual int32 PushAudio(const FPushedAudioDetails & InDetails, TArrayView<const float> InFloat32Interleave) override;
 
 		/** This should not be called. It removes 16 bit PCM samples from the buffer
 		 * which is an unsupported operation of this class.
 		 */
-		virtual int32 PopAudio(TArrayView<int16> InExternalInt16Buffer, FPushedAudioDetails & OutDetails) override;
+		AUDIOCODECENGINE_API virtual int32 PopAudio(TArrayView<int16> InExternalInt16Buffer, FPushedAudioDetails & OutDetails) override;
 
 		/** Copy samples to OutBuffer and remove them from this objects internal
 		 * buffer.
@@ -77,7 +77,7 @@ namespace SoundWaveProxyReaderPrivate
 		 *
 		 * @return The actual number of samples copied.
 		 */
-		virtual int32 PopAudio(TArrayView<float> OutBuffer, FPushedAudioDetails & OutDetails) override;
+		AUDIOCODECENGINE_API virtual int32 PopAudio(TArrayView<float> OutBuffer, FPushedAudioDetails & OutDetails) override;
 
 		/** Remove samples from the internal buffer.
 		 *
