Fix audio sync issues.

Properly report the number of samples that have been "streamed".

Update it in all of our decoders that implement this method.

Still using the "BackCompat" decoder, take advantage of this number to mantain sync with the audio.

diff --git a/Engine/Source/Runtime/AudioCodecEngine/Private/DecoderBackCompat.cpp b/Engine/Source/Runtime/AudioCodecEngine/Private/DecoderBackCompat.cpp
--- a/Engine/Source/Runtime/AudioCodecEngine/Private/DecoderBackCompat.cpp
+++ b/Engine/Source/Runtime/AudioCodecEngine/Private/DecoderBackCompat.cpp
@@ -85,17 +85,26 @@ namespace Audio
 		uint32 BuffSizeInFrames = Reqs.NumSampleFramesWanted;
 		uint8* Buff = (uint8*)ResidualBuffer.GetData();
 
-		int32 NumBytesStreamed = BuffSizeInBytes;
+		int32 NumBytesStreamed = 0;
 		while (!bFinished && NumFramesRemaining > 0)
 		{
 			if (BackCompatSrc.Wave->IsStreaming())
 			{
+				NumBytesStreamed = 0;
 				bFinished = Info->StreamCompressedData(Buff, bLoop, BuffSizeInBytes, NumBytesStreamed);
 			}
 			else
 			{
+				NumBytesStreamed = BuffSizeInBytes;
 				bFinished = Info->ReadCompressedData(Buff, bLoop, BuffSizeInBytes);
 			}
+
+			if (NumBytesStreamed == 0)
+			{
+				// early out so we don't hold up down stream rendering
+				break;
+			}
+
 			int32 NumSamplesStreamed = NumBytesStreamed / sizeof(int16);
 			int32 NumFramesStreamed = NumSamplesStreamed / Desc.NumChannels;
 			PushedDetails.SampleFramesStartOffset = FrameOffset;
@@ -108,6 +117,7 @@ namespace Audio
 
 			FrameOffset += NumFramesStreamed;
 			NumFramesRemaining -= FMath::Min(NumFramesStreamed, NumFramesRemaining);
+			BuffSizeInBytes -= NumBytesStreamed;
 		}
 
 		if (!bFinished)
diff --git a/Engine/Source/Runtime/AudioCodecEngine/Private/SoundWaveProxyReader.cpp b/Engine/Source/Runtime/AudioCodecEngine/Private/SoundWaveProxyReader.cpp
--- a/Engine/Source/Runtime/AudioCodecEngine/Private/SoundWaveProxyReader.cpp
+++ b/Engine/Source/Runtime/AudioCodecEngine/Private/SoundWaveProxyReader.cpp
@@ -339,16 +339,31 @@ int32 FSoundWaveProxyReader::PopAudio(Audio::FAlignedFloatBuffer& OutBuffer)
 
 		bDecoderOutputHasMoreData = DecoderOutput.Num() > 0;
 		bCanProduceMoreAudio = bDecoderOutputHasMoreData || bDecoderCanDecodeMoreData;
+
+		if (bCanProduceMoreAudio && DecoderOutput.Num() == 0)
+		{
+			// we can produce more audio, but we were unable to
+			// this is likely due to the streaming data not being available yet
+			// let's early out to avoid a hitch and hope that it's ready on the next read
+			UE_LOG(LogAudio, Verbose, TEXT("FSoundWaveProxyReader experienced an underrun while decoding samples: %s"), *WaveProxy->GetPackageName().ToString());
+			break;
+		}
 	}
 
 	// Zero pad any unset samples. 
 	if (OutBufferView.Num() > 0)
 	{
-		check(!bCanProduceMoreAudio);
-
 		// Zero out audio that was not set.
+		// advance the current frame index
 		FMemory::Memset(OutBufferView.GetData(), 0, sizeof(float) * OutBufferView.Num());
 		CurrentFrameIndex += OutBufferView.Num() / NumChannels;
+
+		if (Settings.bMaintainAudioSync)
+		{
+			// if we were asked to maintain audio sync, then do some extra book keeping
+			// keep track of the number of samples we'll need to discard the next time we try to read from the decoder
+			NumDecodeSamplesToDiscard += OutBufferView.Num();
+		}
 	}
 
 	return NumSamplesCopied;
@@ -360,6 +375,10 @@ int32 FSoundWaveProxyReader::PopAudioFromDecoderOutput(TArrayView<float> OutBuff
 
 	check(NumChannels > 0);
 
+	int32 NumDiscarded = DecoderOutput.PopAudio(NumDecodeSamplesToDiscard);
+	NumDecodeSamplesToDiscard = NumDecodeSamplesToDiscard - NumDiscarded;
+	check(NumDecodeSamplesToDiscard >= 0);
+
 	int32 NumSamplesCopied = 0;
 	if (DecoderOutput.Num() > 0)
 	{
@@ -456,6 +475,7 @@ bool FSoundWaveProxyReader::InitializeDecoder(float InStartTimeInSeconds)
 	// end hack
 
 	CurrentFrameIndex = FMath::Clamp(static_cast<int32>(InStartTimeInSeconds * GetSampleRate()), 0, GetNumFramesInWave());
+	NumDecodeSamplesToDiscard = 0;
 
 	// Create the decoder
 	Decoder = Codec->CreateDecoder(DecoderInput.Get(), &DecoderOutput);
@@ -491,8 +511,9 @@ bool FSoundWaveProxyReader::InitializeDecoder(float InStartTimeInSeconds)
 	return true;
 }
 
-void FSoundWaveProxyReader::DiscardSamples(int32 InNumSamplesToDiscard)
+int32 FSoundWaveProxyReader::DiscardSamples(int32 InNumSamplesToDiscard)
 {
+	int32 NumSamplesDiscarded = 0;
 	while (InNumSamplesToDiscard > 0)
 	{
 		DecodeResult = Decoder->Decode(false /* bIsLooping */);
@@ -507,6 +528,7 @@ void FSoundWaveProxyReader::DiscardSamples(int32 InNumSamplesToDiscard)
 
 		int32 ActualNumSamplesDiscarded = DecoderOutput.PopAudio(NumSamplesToDiscardThisLoop);
 		InNumSamplesToDiscard -= ActualNumSamplesDiscarded;
+		NumSamplesDiscarded = ActualNumSamplesDiscarded;
 
 		if ((InNumSamplesToDiscard > 0) && (DecodeResult != EDecodeResult::MoreDataRemaining))
 		{
@@ -514,6 +536,8 @@ void FSoundWaveProxyReader::DiscardSamples(int32 InNumSamplesToDiscard)
 			break;
 		}
 	}
+
+	return NumSamplesDiscarded;
 }
 
 float FSoundWaveProxyReader::ClampLoopStartTime(float InStartTimeInSeconds)
diff --git a/Engine/Source/Runtime/AudioCodecEngine/Public/SoundWaveProxyReader.h b/Engine/Source/Runtime/AudioCodecEngine/Public/SoundWaveProxyReader.h
--- a/Engine/Source/Runtime/AudioCodecEngine/Public/SoundWaveProxyReader.h
+++ b/Engine/Source/Runtime/AudioCodecEngine/Public/SoundWaveProxyReader.h
@@ -146,6 +146,7 @@ public:
 		bool bIsLooping = false;
 		float LoopStartTimeInSeconds = 0.f;
 		float LoopDurationInSeconds = 0.f;
+		bool bMaintainAudioSync = false;
 	};
 
 private:
@@ -168,7 +169,7 @@ public:
 	AUDIOCODECENGINE_API void SetIsLooping(bool bInIsLooping);
 
 	/** Returns true if the audio will be looped, false otherwise. */
-	AUDIOCODECENGINE_API FORCEINLINE bool IsLooping() const
+	FORCEINLINE bool IsLooping() const
 	{
 		return Settings.bIsLooping;
 	}
@@ -183,38 +184,38 @@ public:
 	 */
 	AUDIOCODECENGINE_API void SetLoopDuration(float InLoopDurationInSeconds);
 
-	AUDIOCODECENGINE_API FORCEINLINE float GetSampleRate() const
+	FORCEINLINE float GetSampleRate() const
 	{
 		return SampleRate;
 	}
 
-	AUDIOCODECENGINE_API FORCEINLINE int32 GetNumChannels() const
+	FORCEINLINE int32 GetNumChannels() const
 	{
 		return NumChannels;
 	}
 
 	/** Returns the index of the playhead within the complete wave. */
-	AUDIOCODECENGINE_API FORCEINLINE int32 GetFrameIndex() const
+	FORCEINLINE int32 GetFrameIndex() const
 	{
 		return CurrentFrameIndex;
 	}
 
-	AUDIOCODECENGINE_API FORCEINLINE int32 GetNumFramesInWave() const
+	FORCEINLINE int32 GetNumFramesInWave() const
 	{
 		return NumFramesInWave;
 	}
 
-	AUDIOCODECENGINE_API FORCEINLINE int32 GetNumFramesInLoop() const
+	FORCEINLINE int32 GetNumFramesInLoop() const
 	{
 		return LoopEndFrameIndex - LoopStartFrameIndex;
 	}
 
-	AUDIOCODECENGINE_API FORCEINLINE int32 GetLoopStartFrameIndex() const
+	FORCEINLINE int32 GetLoopStartFrameIndex() const
 	{
 		return LoopStartFrameIndex;
 	}
 
-	AUDIOCODECENGINE_API FORCEINLINE int32 GetLoopEndFrameIndex() const
+	FORCEINLINE int32 GetLoopEndFrameIndex() const
 	{
 		return LoopEndFrameIndex;
 	}
@@ -236,7 +237,7 @@ private:
 
 	int32 PopAudioFromDecoderOutput(TArrayView<float> OutBufferView);
 	bool InitializeDecoder(float InStartTimeInSeconds);
-	void DiscardSamples(int32 InNumSamplesToDiscard);
+	int32 DiscardSamples(int32 InNumSamplesToDiscard);
 	float ClampLoopStartTime(float InStartTimeInSeconds);
 	float ClampLoopDuration(float InDurationInSeconds);
 	void UpdateLoopBoundaries();
@@ -252,6 +253,7 @@ private:
 	int32 NumChannels = 0;
 	EDecodeResult DecodeResult = EDecodeResult::MoreDataRemaining;
 	int32 NumFramesInWave = 0;
+	int32 NumDecodeSamplesToDiscard = 0;
 	int32 CurrentFrameIndex = 0;
 	int32 LoopStartFrameIndex = 0;
 	int32 LoopEndFrameIndex = -1;
diff --git a/Engine/Source/Runtime/Engine/Private/ADPCMAudioInfo.cpp b/Engine/Source/Runtime/Engine/Private/ADPCMAudioInfo.cpp
--- a/Engine/Source/Runtime/Engine/Private/ADPCMAudioInfo.cpp
+++ b/Engine/Source/Runtime/Engine/Private/ADPCMAudioInfo.cpp
@@ -771,7 +771,6 @@ bool FADPCMAudioInfo::StreamCompressedData(uint8* Destination, bool bLooping, ui
 
 						// zero out remaining data and bail
 						FMemory::Memset(OutData, 0, BufferSize);
-						OutNumBytesStreamed = OriginalBufferSize;
 						return false;
 					}
 
@@ -896,7 +895,6 @@ bool FADPCMAudioInfo::StreamCompressedData(uint8* Destination, bool bLooping, ui
 					NumConsecutiveReadFailiures++;
 					const bool bReadAttemptTimedOut = NumConsecutiveReadFailiures > ADPCMReadFailiureTimeoutCVar;
 					UE_CLOG(bReadAttemptTimedOut, LogAudio, Warning, TEXT("ADPCM Audio Decode timed out."), bReadAttemptTimedOut);
-					OutNumBytesStreamed = OriginalBufferSize;
 					return NumConsecutiveReadFailiures > ADPCMReadFailiureTimeoutCVar;
 				}
 
diff --git a/Engine/Source/Runtime/Engine/Private/AudioDecompress.cpp b/Engine/Source/Runtime/Engine/Private/AudioDecompress.cpp
--- a/Engine/Source/Runtime/Engine/Private/AudioDecompress.cpp
+++ b/Engine/Source/Runtime/Engine/Private/AudioDecompress.cpp
@@ -244,6 +244,8 @@ bool IStreamedCompressedInfo::StreamCompressedData(uint8* Destination, bool bLoo
 
 	// Write out any PCM data that was decoded during the last request
 	uint32 RawPCMOffset = WriteFromDecodedPCM(Destination, BufferSize);
+	// immediately update the OutNumBytesStreamed to reflect what we just wrote
+	OutNumBytesStreamed = RawPCMOffset;
 
 	// If we have a pending next chunk from seeking, move to it now.
 	if (StreamSeekBlockIndex != INDEX_NONE)
@@ -378,7 +380,6 @@ bool IStreamedCompressedInfo::StreamCompressedData(uint8* Destination, bool bLoo
 			}
 			
 			ZeroBuffer(Destination + RawPCMOffset, BufferSize - RawPCMOffset);
-			OutNumBytesStreamed = BufferSize - RawPCMOffset;
 			return false;
 		}
 	}
@@ -404,14 +405,15 @@ bool IStreamedCompressedInfo::StreamCompressedData(uint8* Destination, bool bLoo
 
 			LastPCMByteSize = 0;
 			ZeroBuffer(Destination + RawPCMOffset, BufferSize - RawPCMOffset);
-			OutNumBytesStreamed = BufferSize - RawPCMOffset;
 			return false;
 		}
 		else
 		{
 			LastPCMByteSize = IncrementCurrentSampleCount(DecodedSamples) * SampleStride;
 
+			// update OutNumBytesStreamed as we write out data
 			RawPCMOffset += WriteFromDecodedPCM(Destination + RawPCMOffset, BufferSize - RawPCMOffset);
+			OutNumBytesStreamed = RawPCMOffset;
 
 			const int32 PreviousChunkIndex = CurrentChunkIndex;
 
@@ -442,8 +444,9 @@ bool IStreamedCompressedInfo::StreamCompressedData(uint8* Destination, bool bLoo
 					}
 					else
 					{
-						RawPCMOffset += ZeroBuffer(Destination + RawPCMOffset, BufferSize - RawPCMOffset);
-						OutNumBytesStreamed = BufferSize - RawPCMOffset;
+						// this is zero padding, so don't include it in the OutNumBytesStreamed count
+						ZeroBuffer(Destination + RawPCMOffset, BufferSize - RawPCMOffset);
+						break;
 					}
 				}
 				else
@@ -470,13 +473,14 @@ bool IStreamedCompressedInfo::StreamCompressedData(uint8* Destination, bool bLoo
 						CurrentChunkIndex, *StreamingSoundWave->GetFName().ToString());
 
 					SrcBufferDataSize = 0;
-					RawPCMOffset += ZeroBuffer(Destination + RawPCMOffset, BufferSize - RawPCMOffset);
+					// this is zero padding, so don't include it in the OutNumBytesStreamed count
+					ZeroBuffer(Destination + RawPCMOffset, BufferSize - RawPCMOffset);
+					break;
 				}
 			}
 		}
 	}
 
-	OutNumBytesStreamed = BufferSize;
 	return bLooped;
 }
 
diff --git a/Engine/Source/Runtime/Engine/Private/VorbisAudioInfo.cpp b/Engine/Source/Runtime/Engine/Private/VorbisAudioInfo.cpp
--- a/Engine/Source/Runtime/Engine/Private/VorbisAudioInfo.cpp
+++ b/Engine/Source/Runtime/Engine/Private/VorbisAudioInfo.cpp
@@ -658,6 +658,7 @@ const uint8* FVorbisAudioInfo::GetLoadedChunk(FSoundWaveProxyPtr InSoundWave, ui
 bool FVorbisAudioInfo::StreamCompressedData(uint8* InDestination, bool bLooping, uint32 BufferSize, int32& OutNumBytesStreamed)
 {
 	const uint32 DestinationSize = BufferSize;
+	OutNumBytesStreamed = 0;
 
 	if (!bDllLoaded)
 	{
@@ -698,7 +699,6 @@ bool FVorbisAudioInfo::StreamCompressedData(uint8* InDestination, bool bLooping,
 			{
 				// zero out the rest of the buffer
 				FMemory::Memzero(InDestination, BufferSize);
-				OutNumBytesStreamed = DestinationSize;
 				return false;
 			}
 
@@ -725,7 +725,6 @@ bool FVorbisAudioInfo::StreamCompressedData(uint8* InDestination, bool bLooping,
 			{
 				// Here we clear out the remainder of the buffer and bail.
 				FMemory::Memzero(InDestination, BufferSize);
-				BytesActuallyRead = BufferSize;
 				break;
 			}
 
