Get rid of BackCompat decoder from WaveProxyReader

Largely just a "copy-paste" implementation of the BackCompat decoder in the SoundWaveProxyReader. Further refactor work will be done in a separate check-in to make it easier to see what actually got changed.

Some things that did change when porting over:
- Remove some "EDecodeResult" enums to avoid complex logic
- Actually use `ICompressedAudioInfo::SeekToTime` instead of reinitializing the decoder to seek.

diff --git a/Engine/Source/Runtime/AudioCodecEngine/Private/SoundWaveProxyReader.cpp b/Engine/Source/Runtime/AudioCodecEngine/Private/SoundWaveProxyReader.cpp
--- a/Engine/Source/Runtime/AudioCodecEngine/Private/SoundWaveProxyReader.cpp
+++ b/Engine/Source/Runtime/AudioCodecEngine/Private/SoundWaveProxyReader.cpp
@@ -5,15 +5,13 @@
 #include "Audio.h"
 #include "Containers/Array.h"
 #include "Containers/ArrayView.h"
-#include "DecoderInputFactory.h"
 #include "DSP/BufferVectorOperations.h"
 #include "DSP/Dsp.h"
 #include "HAL/IConsoleManager.h"
 #include "HAL/LowLevelMemTracker.h"
 #include "HAL/Platform.h"
 #include "HAL/UnrealMemory.h"
-#include "IAudioCodec.h"
-#include "IAudioCodecRegistry.h"
+#include "Interfaces/IAudioFormat.h"
 #include "Math/UnrealMathUtility.h"
 #include "Misc/AssertionMacros.h"
 #include "ProfilingDebugging/CpuProfilerTrace.h"
@@ -29,137 +27,6 @@ FAutoConsoleVariableRef CVarSoundWaveProxyReaderSimulateSeekOnNonSeekable(
 	TEXT("0: Do not simulate seek, !0: Simulate seek"),
 	ECVF_Default);
 
-
-PRAGMA_DISABLE_DEPRECATION_WARNINGS
-
-namespace SoundWaveProxyReaderPrivate
-{
-	/** Construct a FDecoderOutput
-	 *
-	 * @param InNumFramesPerDecode - Maximum number of frames pushed when
-	 *                               audio is decoded. The buffer will be
-	 *                               able to hold twice the number of frames.
-	 */
-	FDecoderOutput::FDecoderOutput(uint32 InNumFramesPerDecode)
-		: NumFramesPerDecode(FMath::Max(InNumFramesPerDecode, MinNumFramesPerDecode))
-	{
-		Init();
-	}
-
-	/** Set the number of channels in the audio being decoded. */
-	void FDecoderOutput::SetNumChannels(uint32 InNumChannels)
-	{
-		NumChannels = FMath::Max(MinNumChannels, InNumChannels);
-		Init();
-	}
-
-	/** Removes all samples from the buffer. */
-	void FDecoderOutput::Reset()
-	{
-		Buffer.SetNum(0);
-	}
-
-	/** Returns requirements used by the audio codec system. */
-	Audio::IDecoderOutput::FRequirements FDecoderOutput::GetRequirements(const Audio::FDecodedFormatInfo& InFormat) const
-	{
-		FRequirements Requirements;
-
-		Requirements.DownstreamFormat = Audio::EBitRepresentation::Float32_Interleaved;
-		Requirements.NumSampleFramesWanted = NumFramesPerDecode;
-		Requirements.NumSampleFramesPerSecond = InFormat.NumFramesPerSec;
-		Requirements.NumChannels = InFormat.NumChannels;
-
-		return Requirements;
-	}
-
-	/** Adds samples to the buffer.
-		*
-		* This is called by the Decoder and should not be called otherwise.
-		*/
-	int32 FDecoderOutput::PushAudio(const Audio::IDecoderOutput::FPushedAudioDetails& InDetails, TArrayView<const int16> In16BitInterleave)
-	{
-		if (InDetails.NumChannels != NumChannels)
-		{
-			SetNumChannels(InDetails.NumChannels);
-		}
-		return PushAudioInternal(InDetails, In16BitInterleave);
-	}
-
-	/** Adds samples to the buffer.
-		*
-		* This is called by the Decoder and should not be called otherwise.
-		*/
-	int32 FDecoderOutput::PushAudio(const FPushedAudioDetails& InDetails, TArrayView<const float> InFloat32Interleave)
-	{
-		if (InDetails.NumChannels != NumChannels)
-		{
-			SetNumChannels(InDetails.NumChannels);
-		}
-		return PushAudioInternal(InDetails, InFloat32Interleave);
-	}
-
-	/** This should not be called. It removes 16 bit PCM samples from the buffer
-		* which is an unsupported operation of this class.
-		*/
-	int32 FDecoderOutput::PopAudio(TArrayView<int16> InExternalInt16Buffer, FPushedAudioDetails& OutDetails)
-	{
-		// This buffer cannot produce 16 bit PCM audio.
-		checkNoEntry();
-		return 0;
-	}
-
-	/** Copy samples to OutBuffer and remove them from this objects internal
-		* buffer.
-		*
-		*
-		* @param OutBuffer - A destination array to copy samples to.
-		* @param OutDetails - Unused.
-		*
-		* @return The actual number of samples copied.
-		*/
-	int32 FDecoderOutput::PopAudio(TArrayView<float> OutBuffer, FPushedAudioDetails& OutDetails)
-	{
-		return Buffer.Pop(OutBuffer.GetData(), OutBuffer.Num());
-	}
-
-	// Initialize buffer size.
-	void FDecoderOutput::Init()
-	{
-		// Allow it to hold two decode buffers max.
-		uint32 MinBufferCapacity = NumChannels * NumFramesPerDecode * 2;
-		constexpr bool bRetainExistingSamples = false;
-		Buffer.Reserve(MinBufferCapacity, bRetainExistingSamples);
-	}
-
-	int32 FDecoderOutput::PushAudioInternal(const FPushedAudioDetails& InDetails, TArrayView<const float> InBuffer)
-	{
-		return Buffer.Push(InBuffer.GetData(), InBuffer.Num());
-	}
-
-	int32 FDecoderOutput::PushAudioInternal(const FPushedAudioDetails& InDetails, TArrayView<const int16> InBuffer)
-	{
-		SampleConversionBuffer.SetNumUninitialized(InBuffer.Num(), false /* bAllowShrinking */);
-
-		// Convert 16 bit pcm to 32 bit float.
-		constexpr float Scalar = 1.f / 32768.f;
-		const int16* Src = InBuffer.GetData();
-		float* Dst = SampleConversionBuffer.GetData();
-		int32 Num = InBuffer.Num();
-
-		// Convert 1 sample at a time, slow.
-		for (int32 i = 0; i < Num; ++i)
-		{
-			*Dst++ = static_cast<float>(*Src++) * Scalar;
-		}
-
-		return PushAudioInternal(InDetails, SampleConversionBuffer);
-	}
-
-}
-
-
-
-
 uint32 FSoundWaveProxyReader::ConformDecodeSize(uint32 InMaxDesiredDecodeSizeInFrames)
 {
 	static_assert(DefaultMinDecodeSizeInFrames >= DecodeSizeQuantizationInFrames, "Min decode size less than decode size quantization");
@@ -175,13 +42,12 @@ uint32 FSoundWaveProxyReader::ConformDecodeSize(uint32 InMaxDesiredDecodeSizeInF
 }
 
 /** Construct a wave proxy reader.
-	*
-	* @param InWaveProxy - A TSharedRef of a FSoundWaveProxy which is to be played.
-	* @param InSettings - Reader settings.
-	*/
+ *
+ * @param InWaveProxy - A TSharedRef of a FSoundWaveProxy which is to be played.
+ * @param InSettings - Reader settings.
+ */
 FSoundWaveProxyReader::FSoundWaveProxyReader(FSoundWaveProxyRef InWaveProxy, const FSettings& InSettings)
 	: WaveProxy(InWaveProxy)
-	, DecoderOutput(ConformDecodeSize(InSettings.MaxDecodeSizeInFrames))
 	, Settings(InSettings)
 {
 	// Get local copies of some values from the proxy. 
@@ -202,17 +68,16 @@ FSoundWaveProxyReader::FSoundWaveProxyReader(FSoundWaveProxyRef InWaveProxy, con
 
 	// Determine max size of decode buffer
 	Settings.MaxDecodeSizeInFrames = FMath::Max(DefaultMinDecodeSizeInFrames, Settings.MaxDecodeSizeInFrames);
-	DecoderOutput.SetNumChannels(NumChannels);
 
 	// Prepare to read audio
 	bIsDecoderValid = InitializeDecoder(Settings.StartTimeInSeconds);
 }
 
 /** Create a wave proxy reader.
-	*
-	* @param InWaveProxy - A TSharedRef of a FSoundWaveProxy which is to be played.
-	* @param InSettings - Reader settings.
-	*/
+ *
+ * @param InWaveProxy - A TSharedRef of a FSoundWaveProxy which is to be played.
+ * @param InSettings - Reader settings.
+ */
 TUniquePtr<FSoundWaveProxyReader> FSoundWaveProxyReader::Create(FSoundWaveProxyRef InWaveProxy, const FSettings& InSettings)
 {
 	if (InWaveProxy->GetSampleRate() <= 0.f)
@@ -259,7 +124,7 @@ void FSoundWaveProxyReader::SetLoopStartTime(float InLoopStartTimeInSeconds)
 }
 
 /** Sets the duration of the loop in seconds. If the value is negative, the
-	* loop duration consists of the entire file. */
+ * loop duration consists of the entire file. */
 void FSoundWaveProxyReader::SetLoopDuration(float InLoopDurationInSeconds)
 {
 	InLoopDurationInSeconds = ClampLoopDuration(InLoopDurationInSeconds);
@@ -272,14 +137,30 @@ void FSoundWaveProxyReader::SetLoopDuration(float InLoopDurationInSeconds)
 
 bool FSoundWaveProxyReader::SeekToTime(float InSeconds)
 {
+	int32 InFrameIndex = FMath::Clamp(static_cast<int32>(InSeconds * GetSampleRate()), 0, GetNumFramesInWave());
+	// ignore seek request if we're already at the specified time
+	if (InFrameIndex == CurrentFrameIndex)
+	{
+		return bIsDecoderValid;
+	}
+
+	if (WaveProxy->IsSeekable() && CompressedAudioInfo)
+	{
+		CompressedAudioInfo->SeekToTime(InSeconds);
+		CurrentFrameIndex = InFrameIndex;
+		DecoderOutput.SetNum(0);
+		NumDecodeSamplesToDiscard = 0;
+		DecodeResult = EDecodeResult::MoreDataRemaining;
+		return bIsDecoderValid;
+	}
 	// Direct seeking is not supported. A new decoder must be created. 
 	bIsDecoderValid = InitializeDecoder(InSeconds);
 	return bIsDecoderValid;
 }
 
 /** Copies audio into OutBuffer. It returns the number of samples copied.
-	* Samples not written to will be set to zero.
-	*/
+ * Samples not written to will be set to zero.
+ */
 int32 FSoundWaveProxyReader::PopAudio(Audio::FAlignedFloatBuffer& OutBuffer)
 {
 	using namespace Audio;
@@ -323,7 +204,7 @@ int32 FSoundWaveProxyReader::PopAudio(Audio::FAlignedFloatBuffer& OutBuffer)
 				if ((!bDecoderCanDecodeMoreData) || (CurrentFrameIndex >= LoopEndFrameIndex))
 				{
 					SeekToTime(Settings.LoopStartTimeInSeconds);
-					bDecoderCanDecodeMoreData = bIsDecoderValid && (EDecodeResult::MoreDataRemaining == DecodeResult);
+					bDecoderCanDecodeMoreData = bIsDecoderValid && (EDecodeResult::Fail != DecodeResult);
 				}
 			}
 		}
@@ -331,9 +212,7 @@ int32 FSoundWaveProxyReader::PopAudio(Audio::FAlignedFloatBuffer& OutBuffer)
 		// Determine if we can / should decode more data. 
 		if ((NumSamplesUnset > 0) && (DecoderOutput.Num() == 0) && bDecoderCanDecodeMoreData)
 		{
-			// Looping is handle within the FSoundWaveProxyReader instead of 
-			// within the decoder.
-			DecodeResult = Decoder->Decode(false /* bIsLooping */);
+			DecodeResult = Decode();
 			bDecoderCanDecodeMoreData = EDecodeResult::MoreDataRemaining == DecodeResult;
 		}
 
@@ -371,11 +250,9 @@ int32 FSoundWaveProxyReader::PopAudio(Audio::FAlignedFloatBuffer& OutBuffer)
 
 int32 FSoundWaveProxyReader::PopAudioFromDecoderOutput(TArrayView<float> OutBufferView)
 {
-	using namespace SoundWaveProxyReaderPrivate;
-
 	check(NumChannels > 0);
 
-	int32 NumDiscarded = DecoderOutput.PopAudio(NumDecodeSamplesToDiscard);
+	int32 NumDiscarded = DecoderOutput.Pop(NumDecodeSamplesToDiscard);
 	NumDecodeSamplesToDiscard = NumDecodeSamplesToDiscard - NumDiscarded;
 	check(NumDecodeSamplesToDiscard >= 0);
 
@@ -383,9 +260,8 @@ int32 FSoundWaveProxyReader::PopAudioFromDecoderOutput(TArrayView<float> OutBuff
 	if (DecoderOutput.Num() > 0)
 	{
 		// Get samples from the decoder buffer.
-		FDecoderOutput::FPushedAudioDetails UnusedDetails;
-		NumSamplesCopied = DecoderOutput.PopAudio(OutBufferView, UnusedDetails);
-
+		NumSamplesCopied = DecoderOutput.Pop(OutBufferView.GetData(), OutBufferView.Num());
+		
 		int32 NumFramesCopied = NumSamplesCopied / NumChannels;
 		CurrentFrameIndex += NumFramesCopied;
 
@@ -406,7 +282,7 @@ int32 FSoundWaveProxyReader::PopAudioFromDecoderOutput(TArrayView<float> OutBuff
 
 				// Remove any remaining samples in the decoder because they 
 				// are past the end of the loop.
-				DecoderOutput.Reset();
+				DecoderOutput.SetNum(0);
 			}
 
 		}
@@ -421,77 +297,81 @@ bool FSoundWaveProxyReader::InitializeDecoder(float InStartTimeInSeconds)
 
 	TRACE_CPUPROFILER_EVENT_SCOPE(FSoundWaveProxyReader::InitializeDecoder);
 
-	// Create decoder input
 	FName Format = WaveProxy->GetRuntimeFormat();
-	DecoderInput = Audio::CreateBackCompatDecoderInput(Format, WaveProxy);
-	if (!DecoderInput.IsValid())
+	IAudioInfoFactory* Factory = IAudioInfoFactoryRegistry::Get().Find(Format);
+	if (!ensure(Factory))
 	{
-		UE_LOG(LogAudio, Error, TEXT("Failed to create decoder input (format:%s) for wave (package:%s)"), *Format.ToString(), *WaveProxy->GetPackageName().ToString());
+		UE_LOG(LogAudio, Error, TEXT("Failed to CompressedAudioInfo for wave (package: %s). Unable to find AudioInfoFactory for format: %s"), *WaveProxy->GetPackageName().ToString(), *Format.ToString());
 		return false;
 	}
 
-	// Seek input to start time.
-	if (!FMath::IsNearlyEqual(0.0f, InStartTimeInSeconds))
+	TUniquePtr<ICompressedAudioInfo> InfoInstance;
+	InfoInstance.Reset(Factory->Create());
+
+	if (!ensure(InfoInstance.IsValid()))
 	{
-		if (WaveProxy->IsSeekable())
+		UE_LOG(LogAudio, Error, TEXT("Failed to created CompressedAudioInfo for wave (package: %s). Unable to create info from factory for for format: %s"), *WaveProxy->GetPackageName().ToString(), *Format.ToString());
+		return false;
+	}
+
+	FSoundQualityInfo Info;
+	if (WaveProxy->IsStreaming())
+	{
+		if (!InfoInstance->StreamCompressedInfo(WaveProxy, &Info))
 		{
-			const bool bSeekSucceeded = DecoderInput->SeekToTime(InStartTimeInSeconds);
-			if (!bSeekSucceeded)
-			{
-				UE_LOG(LogAudio, Warning, TEXT("Failed to seek decoder input during initialization: (format:%s) for wave (package:%s) to time '%.6f'"),
-					*Format.ToString(),
-					*WaveProxy->GetPackageName().ToString(),
-					InStartTimeInSeconds);
-			}
+			UE_LOG(LogAudio, Error, TEXT("Failed to created CompressedAudioInfo for wave (package: %s). Unable to stream compressed info for streaming wave"), *WaveProxy->GetPackageName().ToString());
+			return false;
 		}
-		else
+	}
+	else
+	{
+		if (!InfoInstance->ReadCompressedInfo(WaveProxy->GetResourceData(), WaveProxy->GetResourceSize(), &Info))
 		{
-			UE_LOG(LogAudio, Warning, TEXT("Attempt to seek on non-seekable wave: (format:%s) for wave (package:%s) to time '%.6f'"),
-				*Format.ToString(),
-				*WaveProxy->GetPackageName().ToString(),
-				InStartTimeInSeconds);
+			UE_LOG(LogAudio, Error, TEXT("Failed to created decoder input for wave (package: %s). Unable to read compressed info for non-streaming wave"), *WaveProxy->GetPackageName().ToString());
+			return false;
 		}
 	}
 
-	// Get codec ptr by reading the header info from the decoder input.
-	ICodecRegistry::FCodecPtr Codec = ICodecRegistry::Get().FindCodecByParsingInput(DecoderInput.Get());
-	if (nullptr == Codec)
-	{
-		UE_LOG(LogAudio, Error, TEXT("Failed to find codec (format:%s) for wave (package:%s)"), *Format.ToString(), *WaveProxy->GetPackageName().ToString());
-		return false;
-	}
+	CompressedAudioInfo.Reset(InfoInstance.Release());
 
 	// Read the sample rate and number of frames from the header 
 	// Similar to refreshing the wave data in FMixerBuffer::CreateStreamingBuffer
 	// This is a runtime hack to address incorrect sample rate on soundwaves 
 	// on platforms with Resample for Device enabled (UE-183237)
-	Audio::FFormatDescriptorSection FormatDesc;
-	DecoderInput->FindSection(FormatDesc);
-	SampleRate = FormatDesc.NumFramesPerSec;
-	if (FormatDesc.NumFrames > 0)
+	SampleRate = Info.SampleRate;
+	uint32 NumFrames = (uint32)((float)Info.Duration * Info.SampleRate);
+	if (NumFrames > 0)
 	{
-		NumFramesInWave = FormatDesc.NumFrames;
+		NumFramesInWave = NumFrames;
 	}
 	// end hack
 
-	CurrentFrameIndex = FMath::Clamp(static_cast<int32>(InStartTimeInSeconds * GetSampleRate()), 0, GetNumFramesInWave());
-	NumDecodeSamplesToDiscard = 0;
-
-	// Create the decoder
-	Decoder = Codec->CreateDecoder(DecoderInput.Get(), &DecoderOutput);
-	if (!Decoder.IsValid())
-	{
-		UE_LOG(LogAudio, Error, TEXT("Failed to create decoder (format:%s) for wave (package:%s)"), *Format.ToString(), *WaveProxy->GetPackageName().ToString());
-		DecodeResult = EDecodeResult::Fail;
-		return false;
-	}
-	else
+	// Seek input to start time.
+	if (!FMath::IsNearlyEqual(0.0f, InStartTimeInSeconds))
 	{
-		// The DecodeResult needs to be set to a valid state incase the prior
-		// decoder finished or failed. 
-		DecodeResult = EDecodeResult::MoreDataRemaining;
+		if (WaveProxy->IsSeekable())
+		{
+			CompressedAudioInfo->SeekToTime(InStartTimeInSeconds);
+		}
+		else
+		{
+			UE_LOG(LogAudio, Warning, TEXT("Attempt to seek on non-seekable wave: (format:%s) for wave (package:%s) to time '%.6f'"),
+				*Format.ToString(),
+				*WaveProxy->GetPackageName().ToString(),
+				InStartTimeInSeconds);
+		}
 	}
 
+	CurrentFrameIndex = FMath::Clamp(static_cast<int32>(InStartTimeInSeconds * GetSampleRate()), 0, GetNumFramesInWave());
+
+	// initialized decode buffers
+	const uint32 DecodeSize = ConformDecodeSize(Settings.MaxDecodeSizeInFrames);
+	NumFramesPerDecode = FMath::Max(DecodeSize, NumFramesPerDecode);
+	ResidualBuffer.SetNum(NumFramesPerDecode * NumChannels);
+	SampleConversionBuffer.SetNumUninitialized(NumFramesPerDecode * NumChannels);
+	DecoderOutput.Reserve(/* MinCapacity = */ NumFramesPerDecode * NumChannels * 2, /* bRetainExistingSamples = */ false);
+	NumDecodeSamplesToDiscard = 0;
+
 	// For non-seekable streaming waves, use a fallback method to seek
 	// to the start time. 
 	const bool bUseFallbackSeekMethod = (GSoundWaveProxyReaderSimulateSeekOnNonSeekable != 0) && (InStartTimeInSeconds != 0.f) && (!WaveProxy->IsSeekable());
@@ -507,7 +387,19 @@ bool FSoundWaveProxyReader::InitializeDecoder(float InStartTimeInSeconds)
 		DiscardSamples(NumSamplesToDiscard);
 	}
 
-	// return true if all the components were successfully create
+	if (!CompressedAudioInfo.IsValid())
+	{
+		UE_LOG(LogAudio, Error, TEXT("Failed to create decoder (format:%s) for wave (package:%s)"), *Format.ToString(), *WaveProxy->GetPackageName().ToString());
+		DecodeResult = EDecodeResult::Fail;
+		return false;
+	}
+	else
+	{
+		// The DecodeResult needs to be set to a valid state in case the prior
+		// decoder finished or failed. 
+		DecodeResult = EDecodeResult::MoreDataRemaining;
+	}
+
 	return true;
 }
 
@@ -516,9 +408,9 @@ int32 FSoundWaveProxyReader::DiscardSamples(int32 InNumSamplesToDiscard)
 	int32 NumSamplesDiscarded = 0;
 	while (InNumSamplesToDiscard > 0)
 	{
-		DecodeResult = Decoder->Decode(false /* bIsLooping */);
+		DecodeResult = Decode();
 
-		int32 NumSamplesToDiscardThisLoop = FMath::Min(DecoderOutput.Num(), InNumSamplesToDiscard);
+		int32 NumSamplesToDiscardThisLoop = FMath::Min((int32)DecoderOutput.Num(), InNumSamplesToDiscard);
 
 		if (NumSamplesToDiscardThisLoop < 1)
 		{
@@ -526,7 +418,7 @@ int32 FSoundWaveProxyReader::DiscardSamples(int32 InNumSamplesToDiscard)
 			break;
 		}
 
-		int32 ActualNumSamplesDiscarded = DecoderOutput.PopAudio(NumSamplesToDiscardThisLoop);
+		int32 ActualNumSamplesDiscarded = DecoderOutput.Pop(NumSamplesToDiscardThisLoop);
 		InNumSamplesToDiscard -= ActualNumSamplesDiscarded;
 		NumSamplesDiscarded = ActualNumSamplesDiscarded;
 
@@ -573,4 +465,71 @@ void FSoundWaveProxyReader::UpdateLoopBoundaries()
 	}
 }
 
-PRAGMA_ENABLE_DEPRECATION_WARNINGS
+FSoundWaveProxyReader::EDecodeResult FSoundWaveProxyReader::Decode()
+{
+	check(CompressedAudioInfo.IsValid());
+
+	bool bFinished = false;
+
+	int32 NumFramesRemaining = NumFramesPerDecode;
+	uint32 BuffSizeInBytes = ResidualBuffer.Num() * sizeof(int16);
+	uint32 BuffSizeInFrames = NumFramesPerDecode;
+	uint8* Buff = (uint8*)ResidualBuffer.GetData();
+
+	// cache the streaming flag off the wave
+	// if it has changed since the last Decode() call, bail
+	// something has probably changed in editor
+	if (bIsFirstDecode)
+	{
+		bIsFirstDecode = false;
+	}
+	else
+	{
+		if (bPreviousIsStreaming != WaveProxy->IsStreaming())
+		{
+			return EDecodeResult::Finished;
+		}
+	}
+	bPreviousIsStreaming = WaveProxy->IsStreaming();
+
+	int32 NumBytesStreamed = 0;
+	while (!bFinished && NumFramesRemaining > 0)
+	{
+		if (WaveProxy->IsStreaming())
+		{
+			NumBytesStreamed = 0;
+			bFinished = CompressedAudioInfo->StreamCompressedData(Buff, false, BuffSizeInBytes, NumBytesStreamed);
+		}
+		else
+		{
+			NumBytesStreamed = BuffSizeInBytes;
+			bFinished = CompressedAudioInfo->ReadCompressedData(Buff, false, BuffSizeInBytes);
+		}
+
+		if (NumBytesStreamed == 0)
+		{
+			break;
+		}
+
+		int32 NumSamplesStreamed = NumBytesStreamed / sizeof(int16);
+		int32 NumFramesStreamed = NumSamplesStreamed / NumChannels;
+
+		constexpr float Scalar = 1.f / 32768.f;
+		for (int32 SampleIdx = 0; SampleIdx < NumSamplesStreamed; ++SampleIdx)
+		{
+			SampleConversionBuffer[SampleIdx] = static_cast<float>(ResidualBuffer[SampleIdx]) * Scalar;
+		}
+
+		const float* SampleData = SampleConversionBuffer.GetData();
+		DecoderOutput.Push(SampleData, NumSamplesStreamed);
+
+		NumFramesRemaining -= FMath::Min(NumFramesStreamed, NumFramesRemaining);
+	}
+
+	if (!bFinished)
+	{
+		return EDecodeResult::MoreDataRemaining;
+	}
+
+	return EDecodeResult::Finished;
+}
\ No newline at end of file
diff --git a/Engine/Source/Runtime/AudioCodecEngine/Public/SoundWaveProxyReader.h b/Engine/Source/Runtime/AudioCodecEngine/Public/SoundWaveProxyReader.h
--- a/Engine/Source/Runtime/AudioCodecEngine/Public/SoundWaveProxyReader.h
+++ b/Engine/Source/Runtime/AudioCodecEngine/Public/SoundWaveProxyReader.h
@@ -3,6 +3,7 @@
 #pragma once
 
 #include "DSP/BufferVectorOperations.h"
+#include "AudioDecompress.h"
 #include "HAL/Platform.h"
 #include "IAudioCodec.h"
 #include "Templates/SharedPointer.h"
@@ -11,99 +12,6 @@
 // Forward declare
 class FSoundWaveProxy;
 
-
-PRAGMA_DISABLE_DEPRECATION_WARNINGS
-
-namespace SoundWaveProxyReaderPrivate
-{
-	/** FDecoderOutput maintains a circular buffer of audio which is produced
-	 * from the decoder.
-	 */
-	class UE_DEPRECATED(5.3, "Please use IAudioFormat/IStreamedCompressedInfo instead") FDecoderOutput : public Audio::IDecoderOutput
-	{
-		static constexpr uint32 MinNumFramesPerDecode = 1;
-		static constexpr uint32 DefaultNumChannels = 1;
-		static constexpr uint32 MinNumChannels = 1;
-
-	public:
-		using FPushedAudioDetails = Audio::IDecoderOutput::FPushedAudioDetails;
-
-		/** Construct a FDecoderOutput
-		 *
-		 * @param InNumFramesPerDecode - Maximum number of frames pushed when
-		 *                               audio is decoded. The buffer will be
-		 *                               able to hold twice the number of frames.
-		 */
-		FDecoderOutput(uint32 InNumFramesPerDecode);
-
-		/** Set the number of channels in the audio being decoded. */
-		void SetNumChannels(uint32 InNumChannels);
-
-		/** Returns the number of samples in the buffer. */
-		FORCEINLINE int32 Num() const
-		{
-			return Buffer.Num();
-		}
-
-		/** Removes all samples from the buffer. */
-		void Reset();
-
-		/** Returns requirements used by the audio codec system. */
-		AUDIOCODECENGINE_API virtual Audio::IDecoderOutput::FRequirements GetRequirements(const Audio::FDecodedFormatInfo & InFormat) const override;
-
-		/** Adds samples to the buffer.
-		 *
-		 * This is called by the Decoder and should not be called otherwise.
-		 */
-		AUDIOCODECENGINE_API virtual int32 PushAudio(const Audio::IDecoderOutput::FPushedAudioDetails & InDetails, TArrayView<const int16> In16BitInterleave) override;
-
-		/** Adds samples to the buffer.
-		 *
-		 * This is called by the Decoder and should not be called otherwise.
-		 */
-		AUDIOCODECENGINE_API virtual int32 PushAudio(const FPushedAudioDetails & InDetails, TArrayView<const float> InFloat32Interleave) override;
-
-		/** This should not be called. It removes 16 bit PCM samples from the buffer
-		 * which is an unsupported operation of this class.
-		 */
-		AUDIOCODECENGINE_API virtual int32 PopAudio(TArrayView<int16> InExternalInt16Buffer, FPushedAudioDetails & OutDetails) override;
-
-		/** Copy samples to OutBuffer and remove them from this objects internal
-		 * buffer.
-		 *
-		 *
-		 * @param OutBuffer - A destination array to copy samples to.
-		 * @param OutDetails - Unused.
-		 *
-		 * @return The actual number of samples copied.
-		 */
-		AUDIOCODECENGINE_API virtual int32 PopAudio(TArrayView<float> OutBuffer, FPushedAudioDetails & OutDetails) override;
-
-		/** Remove samples from the internal buffer.
-		 *
-		 * @param InNumSamples - The desired number of samples to remove.
-		 *
-		 * @return The actual number of samples removed.
-		 */
-		FORCEINLINE int32 PopAudio(int32 InNumSamples)
-		{
-			return Buffer.Pop(InNumSamples);
-		}
-
-	private:
-
-		// Initialize buffer size.
-		void Init();
-		int32 PushAudioInternal(const FPushedAudioDetails & InDetails, TArrayView<const float> InBuffer);
-		int32 PushAudioInternal(const FPushedAudioDetails & InDetails, TArrayView<const int16> InBuffer);
-
-		uint32 NumFramesPerDecode = MinNumFramesPerDecode;
-		uint32 NumChannels = DefaultNumChannels;
-		Audio::TCircularAudioBuffer<float> Buffer;
-		Audio::FAlignedFloatBuffer SampleConversionBuffer;
-	};
-}
-
 /** FSoundWaveProxyReader reads a FWaveProxy and outputs 32 bit interleaved audio.
 *
 * FSoundWaveProxyReader provides controls for looping and relevant frame index
@@ -114,7 +22,6 @@ class FSoundWaveProxyReader
 
 public:
 	using FSoundWaveProxyRef = TSharedRef<FSoundWaveProxy, ESPMode::ThreadSafe>;
-	using EDecodeResult = Audio::IDecoder::EDecodeResult;
 
 	/** Minimum number of frames to decode per a call to the decoder.  */
 	static constexpr uint32 DefaultMinDecodeSizeInFrames = 128;
@@ -224,17 +131,26 @@ public:
 	 *
 	 * @param InSeconds - The location to seek the playhead
 	 *
-	 * @return true on success, false on failuer.
+	 * @return true on success, false on failure.
 	 */
 	AUDIOCODECENGINE_API bool SeekToTime(float InSeconds);
 
-	/** Copies audio into OutBuffer. It returns the number of samples copied.
+	/** Pops audio from reader and copies audio into OutBuffer. It returns the number of samples copied.
 	 * Samples not written to will be set to zero.
 	 */
 	AUDIOCODECENGINE_API int32 PopAudio(Audio::FAlignedFloatBuffer& OutBuffer);
 
 private:
 
+	enum EDecodeResult
+	{
+		Fail,					// Decoder failed somehow.
+		MoreDataRemaining,		// Data has produced and there's more remaining
+		Finished				// The decoder has reached the end of the wave data
+	};
+
+	EDecodeResult Decode();
+
 	int32 PopAudioFromDecoderOutput(TArrayView<float> OutBufferView);
 	bool InitializeDecoder(float InStartTimeInSeconds);
 	int32 DiscardSamples(int32 InNumSamplesToDiscard);
@@ -243,9 +159,6 @@ private:
 	void UpdateLoopBoundaries();
 
 	FSoundWaveProxyRef WaveProxy;
-	TUniquePtr<Audio::IDecoderInput> DecoderInput;
-	TUniquePtr<Audio::IDecoder> Decoder;
-	SoundWaveProxyReaderPrivate::FDecoderOutput DecoderOutput;
 
 	FSettings Settings;
 
@@ -262,6 +175,17 @@ private:
 
 	bool bIsDecoderValid = false;
 	bool bFallbackSeekMethodWarningLogged = false;
-};
 
-PRAGMA_ENABLE_DEPRECATION_WARNINGS
\ No newline at end of file
+private:
+	mutable TUniquePtr<class ICompressedAudioInfo> CompressedAudioInfo;
+	Audio::TCircularAudioBuffer<float> DecoderOutput;
+
+	TArray<int16> ResidualBuffer;
+	Audio::FAlignedFloatBuffer SampleConversionBuffer;
+
+	static constexpr uint32 MinNumFramesPerDecode = 1;
+
+	bool bIsFirstDecode = true;
+	bool bPreviousIsStreaming = true;
+	uint32 NumFramesPerDecode = MinNumFramesPerDecode;
+};
\ No newline at end of file
