Fixed compilation for ARM Mac builds

diff --git a/Engine/Source/Runtime/AudioCodecEngine/Private/SoundWaveProxyReader.cpp b/Engine/Source/Runtime/AudioCodecEngine/Private/SoundWaveProxyReader.cpp
--- a/Engine/Source/Runtime/AudioCodecEngine/Private/SoundWaveProxyReader.cpp
+++ b/Engine/Source/Runtime/AudioCodecEngine/Private/SoundWaveProxyReader.cpp
@@ -537,7 +537,7 @@ FSoundWaveProxyReader::EDecodeResult FSoundWaveProxyReader::Decode()
 		const int32 NumSamplesStreamed = NumBytesStreamed / sizeof(int16);
 		const int32 NumFramesStreamed = NumSamplesStreamed / NumChannels;
 
-		Audio::ArrayPcm16ToFloat(MakeArrayView(ResidualBuffer, NumSamplesStreamed), MakeArrayView(SampleConversionBuffer, NumSamplesStreamed));
+		Audio::ArrayPcm16ToFloat(MakeArrayView(ResidualBuffer.GetData(), NumSamplesStreamed), MakeArrayView(SampleConversionBuffer.GetData(), NumSamplesStreamed));
 
 		const float* SampleData = SampleConversionBuffer.GetData();
 		DecoderOutput.Push(SampleData, NumSamplesStreamed);
