- Add a useful public function to SoundWaveProxyReader... CanProduceMoreAudio
- Fix issue where CurrentFrameIndex was being advanced even when the decoder had not produced enough samples.

diff --git a/Engine/Source/Runtime/AudioCodecEngine/Private/SoundWaveProxyReader.cpp b/Engine/Source/Runtime/AudioCodecEngine/Private/SoundWaveProxyReader.cpp
--- a/Engine/Source/Runtime/AudioCodecEngine/Private/SoundWaveProxyReader.cpp
+++ b/Engine/Source/Runtime/AudioCodecEngine/Private/SoundWaveProxyReader.cpp
@@ -158,6 +158,13 @@ bool FSoundWaveProxyReader::SeekToTime(float InSeconds)
 	return bIsDecoderValid;
 }
 
+bool FSoundWaveProxyReader::CanProduceMoreAudio() const
+{
+	const bool bDecoderOutputHasMoreData = DecoderOutput.Num() > 0;
+	const bool bDecoderCanDecodeMoreData = bIsDecoderValid && (EDecodeResult::MoreDataRemaining == DecodeResult);
+	return bDecoderOutputHasMoreData || bDecoderCanDecodeMoreData;
+}
+
 AUDIOCODECENGINE_API bool FSoundWaveProxyReader::SeekToFrame(uint32 InFrameNum)
 {
 	// ignore seek request if we're already at the specified time
@@ -256,15 +263,23 @@ int32 FSoundWaveProxyReader::PopAudio(Audio::FAlignedFloatBuffer& OutBuffer)
 	if (OutBufferView.Num() > 0)
 	{
 		// Zero out audio that was not set.
-		// advance the current frame index
 		FMemory::Memset(OutBufferView.GetData(), 0, sizeof(float) * OutBufferView.Num());
-		CurrentFrameIndex += OutBufferView.Num() / NumChannels;
 
 		if (Settings.bMaintainAudioSync)
 		{
+			ensureMsgf(!Settings.bIsLooping, TEXT("Currently can't BOTH loop a wave and have it maintain sync. The code that does the bookkeeping when the decoder underruns is not robust enough to handle that situation."));
+
 			// if we were asked to maintain audio sync, then do some extra book keeping
 			// keep track of the number of samples we'll need to discard the next time we try to read from the decoder
 			NumDecodeSamplesToDiscard += OutBufferView.Num();
+			// and pretend we are advancing through the data even though we aren't...
+			CurrentFrameIndex += OutBufferView.Num() / NumChannels;
+			if (CurrentFrameIndex > GetNumFramesInWave())
+			{
+				// but don't pretend to go passed the end!...
+				CurrentFrameIndex = GetNumFramesInWave();
+			}
+			// Note: We can do the above because later we will NOT advance CurrentFrameIndex for the decoded samples that we discard!)
 		}
 	}
 
@@ -278,7 +293,7 @@ int32 FSoundWaveProxyReader::PopAudioFromDecoderOutput(TArrayView<float> OutBuff
 	int32 NumDiscarded = DecoderOutput.Pop(NumDecodeSamplesToDiscard);
 	NumDecodeSamplesToDiscard = NumDecodeSamplesToDiscard - NumDiscarded;
 	check(NumDecodeSamplesToDiscard >= 0);
-
+	
 	int32 NumSamplesCopied = 0;
 	if (DecoderOutput.Num() > 0)
 	{
diff --git a/Engine/Source/Runtime/AudioCodecEngine/Public/SoundWaveProxyReader.h b/Engine/Source/Runtime/AudioCodecEngine/Public/SoundWaveProxyReader.h
--- a/Engine/Source/Runtime/AudioCodecEngine/Public/SoundWaveProxyReader.h
+++ b/Engine/Source/Runtime/AudioCodecEngine/Public/SoundWaveProxyReader.h
@@ -149,6 +149,12 @@ public:
 	 */
 	AUDIOCODECENGINE_API int32 PopAudio(Audio::FAlignedFloatBuffer& OutBuffer);
 
+	/** Returns TRUE if the reader can produce audio (eg. has a valid decoder, it can still decode, etc.)
+	 *
+	 * @return true on success, false on failure.
+	 */
+	AUDIOCODECENGINE_API bool CanProduceMoreAudio() const;
+
 private:
 
 	enum EDecodeResult
