delete BackCompat decoders

Eventually move SoundWaveProxyReader out of AudioCodecEngine plugin, and delete AudioCodecEngine plugin entirely.

#tests PIE

diff --git a/Engine/Plugins/Runtime/Metasound/Source/MetasoundEngine/Private/MetasoundAudioBus.cpp b/Engine/Plugins/Runtime/Metasound/Source/MetasoundEngine/Private/MetasoundAudioBus.cpp
--- a/Engine/Plugins/Runtime/Metasound/Source/MetasoundEngine/Private/MetasoundAudioBus.cpp
+++ b/Engine/Plugins/Runtime/Metasound/Source/MetasoundEngine/Private/MetasoundAudioBus.cpp
@@ -3,7 +3,6 @@
 #include "MetasoundAudioBus.h"
 
 #include "AudioDevice.h"
-#include "DecoderInputFactory.h"
 
 namespace Metasound
 {
diff --git a/Engine/Plugins/Runtime/Metasound/Source/MetasoundEngine/Private/MetasoundWave.cpp b/Engine/Plugins/Runtime/Metasound/Source/MetasoundEngine/Private/MetasoundWave.cpp
--- a/Engine/Plugins/Runtime/Metasound/Source/MetasoundEngine/Private/MetasoundWave.cpp
+++ b/Engine/Plugins/Runtime/Metasound/Source/MetasoundEngine/Private/MetasoundWave.cpp
@@ -5,7 +5,6 @@
 #include "AudioDevice.h"
 #include "AudioDeviceManager.h"
 #include "ContentStreaming.h"
-#include "DecoderInputFactory.h"
 #include "DSP/ParamInterpolator.h"
 #include "IAudioCodec.h"
 #include "IAudioCodecRegistry.h"
diff --git a/Engine/Plugins/Runtime/Metasound/Source/MetasoundEngine/Private/MetasoundWavePlayerNode.cpp b/Engine/Plugins/Runtime/Metasound/Source/MetasoundEngine/Private/MetasoundWavePlayerNode.cpp
--- a/Engine/Plugins/Runtime/Metasound/Source/MetasoundEngine/Private/MetasoundWavePlayerNode.cpp
+++ b/Engine/Plugins/Runtime/Metasound/Source/MetasoundEngine/Private/MetasoundWavePlayerNode.cpp
@@ -1,6 +1,5 @@
 // Copyright Epic Games, Inc. All Rights Reserved.
 
-#include "DecoderInputFactory.h"
 #include "DSP/BufferVectorOperations.h"
 #include "DSP/ConvertDeinterleave.h"
 #include "DSP/MultichannelBuffer.h"
diff --git a/Engine/Source/Runtime/AudioCodecEngine/Private/AudioCodecEngineModule.cpp b/Engine/Source/Runtime/AudioCodecEngine/Private/AudioCodecEngineModule.cpp
--- a/Engine/Source/Runtime/AudioCodecEngine/Private/AudioCodecEngineModule.cpp
+++ b/Engine/Source/Runtime/AudioCodecEngine/Private/AudioCodecEngineModule.cpp
@@ -3,7 +3,6 @@
 #include "Modules/ModuleManager.h"
 
 #include "IAudioCodecRegistry.h"
-#include "DecoderBackCompat.h"
 
 PRAGMA_DISABLE_DEPRECATION_WARNINGS
 
@@ -12,19 +11,9 @@ class FAudioCodecEngineModule : public IModuleInterface
 public:
 	void StartupModule() override
 	{
-		Audio::ICodecRegistry::Get().RegisterCodec(
-			MakeUnique<Audio::FBackCompatCodec>()
-		);
 	}
 	void ShutdownModule() override
 	{
-		using namespace Audio;
-		if (ICodecRegistry::FCodecPtr Codec = ICodecRegistry::Get().FindCodecByName(
-			FBackCompatCodec::GetDetailsStatic().Name,
-			FBackCompatCodec::GetDetailsStatic().Version))
-		{
-			ICodecRegistry::Get().UnregisterCodec(Codec);
-		}
 	}
 };
 
diff --git a/Engine/Source/Runtime/AudioCodecEngine/Private/DecoderBackCompat.cpp b/Engine/Source/Runtime/AudioCodecEngine/Private/DecoderBackCompat.cpp
deleted file mode 100644
diff --git a/Engine/Source/Runtime/AudioCodecEngine/Private/DecoderBackCompat.h b/Engine/Source/Runtime/AudioCodecEngine/Private/DecoderBackCompat.h
deleted file mode 100644
diff --git a/Engine/Source/Runtime/AudioCodecEngine/Private/DecoderInputBackCompat.cpp b/Engine/Source/Runtime/AudioCodecEngine/Private/DecoderInputBackCompat.cpp
deleted file mode 100644
diff --git a/Engine/Source/Runtime/AudioCodecEngine/Private/DecoderInputBackCompat.h b/Engine/Source/Runtime/AudioCodecEngine/Private/DecoderInputBackCompat.h
deleted file mode 100644
diff --git a/Engine/Source/Runtime/AudioCodecEngine/Private/DecoderInputFactory.cpp b/Engine/Source/Runtime/AudioCodecEngine/Private/DecoderInputFactory.cpp
deleted file mode 100644
diff --git a/Engine/Source/Runtime/AudioCodecEngine/Public/DecoderInputFactory.h b/Engine/Source/Runtime/AudioCodecEngine/Public/DecoderInputFactory.h
deleted file mode 100644
